Rename Shotgun to Shotgrid.  To do this in a single change, this a delete + add for the new files plus a DefaultShotgrid.ini for redirects.

diff --git a/Engine/Plugins/Experimental/Shotgrid/Config/DefaultShotgrid.ini b/Engine/Plugins/Experimental/Shotgrid/Config/DefaultShotgrid.ini
new file mode 100644
--- /dev/null
+++ b/Engine/Plugins/Experimental/Shotgrid/Config/DefaultShotgrid.ini
@@ -0,0 +1,8 @@
+[/Script/Engine.Engine]
+
+[CoreRedirects]
++ClassRedirects=(OldName="Shotgun",NewName="Shotgrid",MatchSubstring=true)
++EnumRedirects=(OldName="Shotgun",NewName="Shotgrid",MatchSubstring=true)
++FunctionRedirects=(OldName="Shotgun",NewName="Shotgrid",MatchSubstring=true)
++StructRedirects=(OldName="Shotgun",NewName="Shotgrid",MatchSubstring=true)
++PackageRedirects=(OldName="Shotgun",NewName="/Shotgrid/",MatchSubstring=true)
\ No newline at end of file
diff --git a/Engine/Plugins/Experimental/Shotgrid/Shotgrid.uplugin b/Engine/Plugins/Experimental/Shotgrid/Shotgrid.uplugin
new file mode 100644
--- /dev/null
+++ b/Engine/Plugins/Experimental/Shotgrid/Shotgrid.uplugin
@@ -0,0 +1,44 @@
+{
+	"FileVersion" : 1,
+	"Version" : 1,
+	"VersionName" : "0.10",
+	"FriendlyName" : "Shotgrid Plugin",
+	"Description" : "Shotgrid integration for the UE4 Editor.",
+	"Category" : "Editor",
+	"CreatedBy" : "Epic Games, Inc.",
+	"CreatedByURL" : "http://epicgames.com",
+	"DocsURL" : "",
+	"MarketplaceURL" : "",
+	"SupportURL" : "",
+	"EnabledByDefault" : false,
+	"CanContainContent" : true,
+	"IsBetaVersion" : true,
+	"Installed" : false,
+	"SupportedTargetPlatforms":
+	[
+		"Win64"
+	],
+	"Modules":
+	[
+		{
+			"Name": "Shotgrid",
+			"Type": "Editor",
+			"LoadingPhase": "Default",
+			"PlatformAllowList":
+			[
+				"Win64"
+			]
+		}
+	],
+	"Plugins":
+	[
+		{
+			"Name": "PythonScriptPlugin",
+			"Enabled": true
+		},
+		{
+			"Name": "EditorScriptingUtilities",
+			"Enabled": true
+		}
+	]
+}
\ No newline at end of file
diff --git a/Engine/Plugins/Experimental/Shotgrid/Source/Shotgrid/Private/ShotgridEngine.cpp b/Engine/Plugins/Experimental/Shotgrid/Source/Shotgrid/Private/ShotgridEngine.cpp
new file mode 100644
--- /dev/null
+++ b/Engine/Plugins/Experimental/Shotgrid/Source/Shotgrid/Private/ShotgridEngine.cpp
@@ -0,0 +1,102 @@
+// Copyright Epic Games, Inc. All Rights Reserved.
+
+#include "ShotgridEngine.h"
+
+#include "AssetRegistryModule.h"
+#include "IAssetRegistry.h"
+
+#include "GameFramework/Actor.h"
+#include "IPythonScriptPlugin.h"
+#include "Misc/CoreDelegates.h"
+#include "Misc/Paths.h"
+#include "Templates/Casts.h"
+#include "UObject/UObjectHash.h"
+
+UShotgridEngine* UShotgridEngine::GetInstance()
+{
+	// The Python Shotgrid Engine instance must come from a Python class derived from UShotgridEngine
+	// In Python, there should be only one derivation, but hot-reloading will create new derived classes, so use the last one
+	TArray<UClass*> ShotgridEngineClasses;
+	GetDerivedClasses(UShotgridEngine::StaticClass(), ShotgridEngineClasses);
+	int32 NumClasses = ShotgridEngineClasses.Num();
+	if (NumClasses > 0)
+	{
+		return Cast<UShotgridEngine>(ShotgridEngineClasses[NumClasses - 1]->GetDefaultObject());
+	}
+	return nullptr;
+}
+
+static void OnEditorExit()
+{
+	if (UShotgridEngine* Engine = UShotgridEngine::GetInstance())
+	{
+		Engine->Shutdown();
+	}
+}
+
+void UShotgridEngine::OnEngineInitialized() const
+{
+	IPythonScriptPlugin::Get()->OnPythonShutdown().AddStatic(OnEditorExit);
+}
+
+void UShotgridEngine::SetSelection(const TArray<FAssetData>* InSelectedAssets, const TArray<AActor*>* InSelectedActors)
+{
+	SelectedAssets.Reset();
+	WeakSelectedActors.Reset();
+
+	if (InSelectedAssets)
+	{
+		SelectedAssets = *InSelectedAssets;
+	}
+	if (InSelectedActors)
+	{
+		// Also set the assets referenced by the selected actors as selected assets
+		IAssetRegistry& AssetRegistry = FModuleManager::LoadModuleChecked<FAssetRegistryModule>("AssetRegistry").Get();
+		TArray<FAssetData> AllReferencedAssets;
+		for (const AActor* Actor : *InSelectedActors)
+		{
+			WeakSelectedActors.Add(Actor);
+
+			TArray<UObject*> ActorAssets = GetReferencedAssets(Actor);
+			for (UObject* Asset : ActorAssets)
+			{
+				if (IsValid(Asset) && Asset->IsAsset())
+				{
+					FAssetData AssetData = AssetRegistry.GetAssetByObjectPath(FName(*Asset->GetPathName()));
+					AllReferencedAssets.AddUnique(MoveTemp(AssetData));
+				}
+			}
+		}
+		SelectedAssets = MoveTemp(AllReferencedAssets);
+	}
+}
+
+TArray<AActor*> UShotgridEngine::GetSelectedActors()
+{
+	TArray<AActor*> Actors;
+	for (const FWeakObjectPtr& ObjPtr : WeakSelectedActors)
+	{
+		if (AActor* Actor = Cast<AActor>(ObjPtr.Get()))
+		{
+			Actors.Add(Actor);
+		}
+	}
+	return Actors;
+}
+
+TArray<UObject*> UShotgridEngine::GetReferencedAssets(const AActor* Actor) const
+{
+	TArray<UObject*> ReferencedAssets;
+
+	if (Actor)
+	{
+		Actor->GetReferencedContentObjects(ReferencedAssets);
+	}
+
+	return ReferencedAssets;
+}
+
+FString UShotgridEngine::GetShotgridWorkDir()
+{
+	return FPaths::ConvertRelativePathToFull(FPaths::ProjectDir());
+}
diff --git a/Engine/Plugins/Experimental/Shotgrid/Source/Shotgrid/Private/ShotgridEngine.h b/Engine/Plugins/Experimental/Shotgrid/Source/Shotgrid/Private/ShotgridEngine.h
new file mode 100644
--- /dev/null
+++ b/Engine/Plugins/Experimental/Shotgrid/Source/Shotgrid/Private/ShotgridEngine.h
@@ -0,0 +1,89 @@
+// Copyright Epic Games, Inc. All Rights Reserved.
+
+#pragma once
+
+#include "CoreMinimal.h"
+#include "AssetData.h"
+
+#include "ShotgridEngine.generated.h"
+
+class AActor;
+
+USTRUCT(Blueprintable)
+struct FShotgridMenuItem
+{
+	GENERATED_BODY()
+
+public:
+	/** Command name for internal use */
+	UPROPERTY(EditAnywhere, BlueprintReadWrite, Category = Python)
+	FString Name;
+
+	/** Text to display in the menu */
+	UPROPERTY(EditAnywhere, BlueprintReadWrite, Category = Python)
+	FString Title;
+
+	/** Description text for the tooltip */
+	UPROPERTY(EditAnywhere, BlueprintReadWrite, Category = Python)
+	FString Description;
+
+	/** Menu item type to help interpret the command */
+	UPROPERTY(EditAnywhere, BlueprintReadWrite, Category = Python)
+	FString Type;
+};
+
+
+/**
+ * Wrapper for the Python Shotgrid Engine
+ * The functions are implemented in Python by a class that derives from this one
+ */
+UCLASS(Blueprintable)
+class UShotgridEngine : public UObject
+{
+	GENERATED_BODY()
+
+public:
+	/** Get the instance of the Python Shotgrid Engine */
+	UFUNCTION(BlueprintCallable, Category = Python)
+	static UShotgridEngine* GetInstance();
+
+	/** Callback for when the Python Shotgrid Engine has finished initialization */
+	UFUNCTION(BlueprintCallable, Category = Python)
+	void OnEngineInitialized() const;
+
+	/** Get the available Shotgrid commands from the Python Shotgrid Engine */
+	UFUNCTION(BlueprintImplementableEvent, Category = Python)
+	TArray<FShotgridMenuItem> GetShotgridMenuItems() const;
+
+	/** Execute a Shotgrid command by name in the Python Shotgrid Engine */
+	UFUNCTION(BlueprintImplementableEvent, Category = Python)
+	void ExecuteCommand(const FString& CommandName) const;
+
+	/** Shut down the Python Shotgrid Engine */
+	UFUNCTION(BlueprintImplementableEvent, Category = Python)
+	void Shutdown() const;
+
+	/** Set the selected objects that will be used to determine the Shotgrid Engine context and execute Shotgrid commands */
+	void SetSelection(const TArray<FAssetData>* InSelectedAssets, const TArray<AActor*>* InSelectedActors);
+
+	/** Get the assets that are referenced by the given Actor */
+	UFUNCTION(BlueprintCallable, Category = Python)
+	TArray<UObject*> GetReferencedAssets(const AActor* Actor) const;
+
+	/** Get the root path for the Shotgrid work area */
+	UFUNCTION(BlueprintCallable, Category = Python)
+	static FString GetShotgridWorkDir();
+
+	/** Selected assets to be used for Shotgrid commands */
+	UPROPERTY(EditAnywhere, BlueprintReadWrite, Category = Python)
+	TArray<FAssetData> SelectedAssets;
+
+	UPROPERTY(EditAnywhere, BlueprintReadWrite, Category = Python, meta = (Deprecated, DeprecationMessage = "SelectedActors is deprecated, use GetSelectedActors instead." ))
+	TArray<AActor*> SelectedActors;
+
+	/** Selected actors to be used for Shotgrid commands */
+	UFUNCTION(BlueprintCallable, Category = Python)
+	TArray<AActor*> GetSelectedActors();
+
+	TArray<FWeakObjectPtr> WeakSelectedActors;
+};
diff --git a/Engine/Plugins/Experimental/Shotgrid/Source/Shotgrid/Private/ShotgridModule.cpp b/Engine/Plugins/Experimental/Shotgrid/Source/Shotgrid/Private/ShotgridModule.cpp
new file mode 100644
--- /dev/null
+++ b/Engine/Plugins/Experimental/Shotgrid/Source/Shotgrid/Private/ShotgridModule.cpp
@@ -0,0 +1,61 @@
+// Copyright Epic Games, Inc. All Rights Reserved.
+
+#include "IShotgridModule.h"
+#include "ShotgridSettings.h"
+#include "ShotgridUIManager.h"
+
+#include "ISettingsModule.h"
+#include "ISettingsSection.h"
+
+#define LOCTEXT_NAMESPACE "Shotgrid"
+
+class FShotgridModule : public IShotgridModule
+{
+public:
+	virtual void StartupModule() override
+	{
+		if (GIsEditor && !IsRunningCommandlet())
+		{
+			RegisterSettings();
+
+			FShotgridUIManager::Initialize();
+		}
+	}
+
+	virtual void ShutdownModule() override
+	{
+		if ((GIsEditor && !IsRunningCommandlet()))
+		{
+			FShotgridUIManager::Shutdown();
+
+			UnregisterSettings();
+		}
+	}
+
+protected:
+	void RegisterSettings()
+	{
+		ISettingsModule* SettingsModule = FModuleManager::GetModulePtr<ISettingsModule>("Settings");
+		if (SettingsModule != nullptr)
+		{
+			ISettingsSectionPtr SettingsSection = SettingsModule->RegisterSettings("Project", "Plugins", "Shotgrid",
+				LOCTEXT("ShotgridSettingsName", "Shotgrid"),
+				LOCTEXT("ShotgridSettingsDescription", "Configure the Shotgrid plug-in."),
+				GetMutableDefault<UShotgridSettings>()
+			);
+		}
+	}
+
+	void UnregisterSettings()
+	{
+		ISettingsModule* SettingsModule = FModuleManager::GetModulePtr<ISettingsModule>("Settings");
+		if (SettingsModule != nullptr)
+		{
+			SettingsModule->UnregisterSettings("Project", "Plugins", "Shotgrid");
+		}
+	}
+};
+
+IMPLEMENT_MODULE(FShotgridModule, Shotgrid);
+
+#undef LOCTEXT_NAMESPACE
diff --git a/Engine/Plugins/Experimental/Shotgrid/Source/Shotgrid/Private/ShotgridSettings.cpp b/Engine/Plugins/Experimental/Shotgrid/Source/Shotgrid/Private/ShotgridSettings.cpp
new file mode 100644
--- /dev/null
+++ b/Engine/Plugins/Experimental/Shotgrid/Source/Shotgrid/Private/ShotgridSettings.cpp
@@ -0,0 +1,91 @@
+// Copyright Epic Games, Inc. All Rights Reserved.
+
+#include "ShotgridSettings.h"
+
+#if WITH_EDITOR
+
+DEFINE_LOG_CATEGORY_STATIC(LogShotgrid, Log, All);
+
+void UShotgridSettings::PostInitProperties()
+{
+	Super::PostInitProperties();
+
+	ApplyMetaDataTagsSettings();
+}
+
+void UShotgridSettings::ApplyMetaDataTagsSettings()
+{
+	TSet<FName>& GlobalTagsForAssetRegistry = UObject::GetMetaDataTagsForAssetRegistry();
+	for (FName Tag : MetaDataTagsForAssetRegistry)
+	{
+		if (!Tag.IsNone())
+		{
+			if (!GlobalTagsForAssetRegistry.Contains(Tag))
+			{
+				GlobalTagsForAssetRegistry.Add(Tag);
+			}
+			else
+			{
+				// To catch the case where the same tag is used by different users and their settings are synced after edition
+				UE_LOG(LogShotgrid, Warning, TEXT("Cannot use duplicate metadata tag '%s' for Asset Registry"), *Tag.ToString());
+			}
+		}
+	}
+}
+
+void UShotgridSettings::ClearMetaDataTagsSettings()
+{
+	TSet<FName>& GlobalTagsForAssetRegistry = UObject::GetMetaDataTagsForAssetRegistry();
+	for (FName Tag : MetaDataTagsForAssetRegistry)
+	{
+		if (!Tag.IsNone())
+		{
+			GlobalTagsForAssetRegistry.Remove(Tag);
+		}
+	}
+}
+
+void UShotgridSettings::PreEditChange(FProperty* PropertyAboutToChange)
+{
+	Super::PreEditChange(PropertyAboutToChange);
+
+	if (PropertyAboutToChange != nullptr)
+	{
+		FName PropertyName = PropertyAboutToChange->GetFName();
+		if (PropertyName == GET_MEMBER_NAME_CHECKED(UShotgridSettings, MetaDataTagsForAssetRegistry))
+		{
+			ClearMetaDataTagsSettings();
+		}
+	}
+}
+
+void UShotgridSettings::PostEditChangeProperty(FPropertyChangedEvent& PropertyChangedEvent)
+{
+	Super::PostEditChangeProperty(PropertyChangedEvent);
+
+	FName PropertyName = (PropertyChangedEvent.MemberProperty != NULL) ? PropertyChangedEvent.MemberProperty->GetFName() : NAME_None;
+	if (PropertyName == GET_MEMBER_NAME_CHECKED(UShotgridSettings, MetaDataTagsForAssetRegistry))
+	{
+		if (PropertyChangedEvent.ChangeType == EPropertyChangeType::ValueSet)
+		{
+			// Check if the new value already exists in the global tags list
+			int32 Index = PropertyChangedEvent.GetArrayIndex(PropertyName.ToString());
+			if (Index > 0)
+			{
+				TSet<FName>::TIterator It = MetaDataTagsForAssetRegistry.CreateIterator();
+				for (int32 i = 0; i < Index; ++i)
+				{
+					++It;
+				}
+				FName NewValue = It ? *It : FName();
+				if (UObject::GetMetaDataTagsForAssetRegistry().Contains(NewValue))
+				{
+					*It = FName();
+					UE_LOG(LogShotgrid, Warning, TEXT("Cannot use duplicate metadata tag '%s' for Asset Registry"), *NewValue.ToString());
+				}
+			}
+		}
+		ApplyMetaDataTagsSettings();
+	}
+}
+#endif
\ No newline at end of file
diff --git a/Engine/Plugins/Experimental/Shotgrid/Source/Shotgrid/Private/ShotgridSettings.h b/Engine/Plugins/Experimental/Shotgrid/Source/Shotgrid/Private/ShotgridSettings.h
new file mode 100644
--- /dev/null
+++ b/Engine/Plugins/Experimental/Shotgrid/Source/Shotgrid/Private/ShotgridSettings.h
@@ -0,0 +1,28 @@
+// Copyright Epic Games, Inc. All Rights Reserved.
+
+#pragma once
+
+#include "CoreMinimal.h"
+
+#include "ShotgridSettings.generated.h"
+
+UCLASS(config = Game)
+class UShotgridSettings : public UObject
+{
+	GENERATED_BODY()
+
+public:
+	/** The metadata tags to be transferred to the Asset Registry. */
+	UPROPERTY(config, EditAnywhere, AdvancedDisplay, Category = "Asset Registry", DisplayName = "Metadata Tags For Asset Registry")
+	TSet<FName> MetaDataTagsForAssetRegistry;
+
+#if WITH_EDITOR
+	virtual void PostInitProperties() override;
+	virtual void PreEditChange(FProperty* PropertyAboutToChange) override;
+	virtual void PostEditChangeProperty(FPropertyChangedEvent& PropertyChangedEvent) override;
+
+private:
+	void ApplyMetaDataTagsSettings();
+	void ClearMetaDataTagsSettings();
+#endif
+};
diff --git a/Engine/Plugins/Experimental/Shotgrid/Source/Shotgrid/Private/ShotgridStyle.cpp b/Engine/Plugins/Experimental/Shotgrid/Source/Shotgrid/Private/ShotgridStyle.cpp
new file mode 100644
--- /dev/null
+++ b/Engine/Plugins/Experimental/Shotgrid/Source/Shotgrid/Private/ShotgridStyle.cpp
@@ -0,0 +1,74 @@
+// Copyright Epic Games, Inc. All Rights Reserved.
+
+#include "ShotgridStyle.h"
+
+#include "Framework/Application/SlateApplication.h"
+#include "Styling/SlateStyle.h"
+#include "Styling/SlateStyleRegistry.h"
+
+TUniquePtr< FSlateStyleSet > FShotgridStyle::ShotgridStyleInstance = nullptr;
+
+void FShotgridStyle::Initialize()
+{
+	if (!ShotgridStyleInstance.IsValid())
+	{
+		ShotgridStyleInstance = Create();
+		FSlateStyleRegistry::RegisterSlateStyle(*ShotgridStyleInstance);
+	}
+}
+
+void FShotgridStyle::Shutdown()
+{
+	if (ShotgridStyleInstance.IsValid())
+	{
+		FSlateStyleRegistry::UnRegisterSlateStyle(*ShotgridStyleInstance);
+		ShotgridStyleInstance.Reset();
+	}
+}
+
+FName FShotgridStyle::GetStyleSetName()
+{
+	static FName StyleSetName(TEXT("ShotgridStyle"));
+	return StyleSetName;
+}
+
+FName FShotgridStyle::GetContextName()
+{
+	static FName ContextName(TEXT("Shotgrid"));
+	return ContextName;
+}
+
+#define IMAGE_BRUSH( RelativePath, ... ) FSlateImageBrush( Style->RootToContentDir( RelativePath, TEXT(".png") ), __VA_ARGS__ )
+
+const FVector2D Icon20x20(20.0f, 20.0f);
+const FVector2D Icon40x40(40.0f, 40.0f);
+
+TUniquePtr< FSlateStyleSet > FShotgridStyle::Create()
+{
+	TUniquePtr< FSlateStyleSet > Style = MakeUnique<FSlateStyleSet>(GetStyleSetName());
+	Style->SetContentRoot(FPaths::EnginePluginsDir() / TEXT("Experimental/Shotgrid/Resources"));
+
+	return Style;
+}
+
+void FShotgridStyle::SetIcon(const FString& StyleName, const FString& ResourcePath)
+{
+	FSlateStyleSet* Style = ShotgridStyleInstance.Get();
+
+	FString Name(GetContextName().ToString());
+	Name = Name + "." + StyleName;
+	Style->Set(*Name, new IMAGE_BRUSH(ResourcePath, Icon40x40));
+
+	Name += ".Small";
+	Style->Set(*Name, new IMAGE_BRUSH(ResourcePath, Icon20x20));
+
+	FSlateApplication::Get().GetRenderer()->ReloadTextureResources();
+}
+
+#undef IMAGE_BRUSH
+
+const ISlateStyle& FShotgridStyle::Get()
+{
+	check(ShotgridStyleInstance);
+	return *ShotgridStyleInstance;
+}
diff --git a/Engine/Plugins/Experimental/Shotgrid/Source/Shotgrid/Private/ShotgridStyle.h b/Engine/Plugins/Experimental/Shotgrid/Source/Shotgrid/Private/ShotgridStyle.h
new file mode 100644
--- /dev/null
+++ b/Engine/Plugins/Experimental/Shotgrid/Source/Shotgrid/Private/ShotgridStyle.h
@@ -0,0 +1,26 @@
+// Copyright Epic Games, Inc. All Rights Reserved.
+
+#pragma once
+
+#include "CoreMinimal.h"
+
+class FSlateStyleSet;
+class ISlateStyle;
+
+/** Manages the style that provides resources for Shotgrid integration */
+class FShotgridStyle
+{
+public:
+	static void Initialize();
+	static void Shutdown();
+	static const ISlateStyle& Get();
+	static FName GetStyleSetName();
+	static FName GetContextName();
+
+	/** Adds an icon (.png in the given ResourcePath) with the given style name to this style set */
+	static void SetIcon(const FString& StyleName, const FString& ResourcePath);
+
+private:
+	static TUniquePtr< FSlateStyleSet > Create();
+	static TUniquePtr< FSlateStyleSet > ShotgridStyleInstance;
+};
\ No newline at end of file
diff --git a/Engine/Plugins/Experimental/Shotgrid/Source/Shotgrid/Private/ShotgridUIManager.cpp b/Engine/Plugins/Experimental/Shotgrid/Source/Shotgrid/Private/ShotgridUIManager.cpp
new file mode 100644
--- /dev/null
+++ b/Engine/Plugins/Experimental/Shotgrid/Source/Shotgrid/Private/ShotgridUIManager.cpp
@@ -0,0 +1,285 @@
+// Copyright Epic Games, Inc. All Rights Reserved.
+
+#include "ShotgridUIManager.h"
+#include "ShotgridEngine.h"
+#include "ShotgridStyle.h"
+
+#include "AssetData.h"
+#include "ContentBrowserModule.h"
+#include "LevelEditor.h"
+#include "Framework/MultiBox/MultiBoxBuilder.h"
+#include "Framework/MultiBox/MultiBoxExtender.h"
+#include "Misc/Paths.h"
+#include "Modules/ModuleManager.h"
+
+#include "IPythonScriptPlugin.h"
+
+#define LOCTEXT_NAMESPACE "Shotgrid"
+#define LEVELEDITOR_MODULE_NAME TEXT("LevelEditor")
+#define CONTENTBROWSER_MODULE_NAME TEXT("ContentBrowser")
+
+TUniquePtr<FShotgridUIManagerImpl> FShotgridUIManager::Instance;
+
+class FShotgridUIManagerImpl
+{
+public:
+	void Initialize();
+	void Shutdown();
+
+private:
+	FDelegateHandle ContentBrowserAssetExtenderDelegateHandle;
+	FDelegateHandle LevelEditorExtenderDelegateHandle;
+
+	bool bIsShotgridEnabled;
+
+	void SetupShotgridMenu();
+	void SetupShotgridContextMenus();
+	void RemoveShotgridContextMenus();
+
+	void FillToolbar(FToolBarBuilder& ToolbarBuilder);
+	TSharedRef<SWidget> GenerateShotgridToolbarMenu();
+	void GenerateShotgridMenuContent(FMenuBuilder& MenuBuilder, const TArray<FAssetData>* SelectedAssets, const TArray< AActor*>* SelectedActors);
+	void GenerateShotgridAssetContextMenu(FMenuBuilder& MenuBuilder, TArray<FAssetData> SelectedAssets);
+	void GenerateShotgridActorContextMenu(FMenuBuilder& MenuBuilder, TArray< AActor*> SelectedActors);
+
+	// Menu extender callbacks
+	TSharedRef<FExtender> OnExtendLevelEditor(const TSharedRef<FUICommandList> CommandList, const TArray<AActor*> SelectedActors);
+	TSharedRef<FExtender> OnExtendContentBrowserAssetSelectionMenu(const TArray<FAssetData>& SelectedAssets);
+};
+
+void FShotgridUIManagerImpl::Initialize()
+{
+	bIsShotgridEnabled = false;
+
+	// Check if the bootstrap environment variable is set and that the script exists
+	FString ShotgridBootstrap = FPlatformMisc::GetEnvironmentVariable(TEXT("UE_SHOTGUN_BOOTSTRAP"));
+	if (!ShotgridBootstrap.IsEmpty() && FPaths::FileExists(ShotgridBootstrap))
+	{
+		// The following environment variables must be set for the Shotgrid apps to be fully functional
+		// These variables are automatically set when the editor is launched through Shotgrid Desktop
+		FString ShotgridEngine = FPlatformMisc::GetEnvironmentVariable(TEXT("SHOTGUN_ENGINE"));
+		FString ShotgridEntityType = FPlatformMisc::GetEnvironmentVariable(TEXT("SHOTGUN_ENTITY_TYPE"));
+		FString ShotgridEntityId = FPlatformMisc::GetEnvironmentVariable(TEXT("SHOTGUN_ENTITY_ID"));
+
+		if (ShotgridEngine == TEXT("tk-unreal") && !ShotgridEntityType.IsEmpty() && !ShotgridEntityId.IsEmpty())
+		{
+			bIsShotgridEnabled = true;
+
+			// Set environment variable in the Python interpreter to enable the Shotgrid Unreal init script
+			IPythonScriptPlugin::Get()->ExecPythonCommand(TEXT("import os\nos.environ['UE_SHOTGUN_ENABLED']='True'"));
+		}
+	}
+
+	if (bIsShotgridEnabled)
+	{
+		FShotgridStyle::Initialize();
+
+		SetupShotgridMenu();
+		SetupShotgridContextMenus();
+	}
+}
+
+void FShotgridUIManagerImpl::Shutdown()
+{
+	if (bIsShotgridEnabled)
+	{
+		RemoveShotgridContextMenus();
+
+		FShotgridStyle::Shutdown();
+	}
+}
+
+void FShotgridUIManagerImpl::SetupShotgridMenu()
+{
+	// Set the Shotgrid icons
+	FShotgridStyle::SetIcon("Logo", "sg_logo_80px");
+	FShotgridStyle::SetIcon("ContextLogo", "sg_context_logo");
+
+	// Add a Shotgrid toolbar section after the settings section of the level editor
+	FLevelEditorModule& LevelEditorModule = FModuleManager::LoadModuleChecked<FLevelEditorModule>(LEVELEDITOR_MODULE_NAME);
+
+	TSharedPtr<FExtender> ToolbarExtender = MakeShareable(new FExtender);
+	ToolbarExtender->AddToolBarExtension("Settings", EExtensionHook::After, nullptr, FToolBarExtensionDelegate::CreateRaw(this, &FShotgridUIManagerImpl::FillToolbar));
+
+	LevelEditorModule.GetToolBarExtensibilityManager()->AddExtender(ToolbarExtender);
+}
+
+void FShotgridUIManagerImpl::SetupShotgridContextMenus()
+{
+	// Register Content Browser menu extender
+	FContentBrowserModule& ContentBrowserModule = FModuleManager::LoadModuleChecked<FContentBrowserModule>(CONTENTBROWSER_MODULE_NAME);
+
+	TArray<FContentBrowserMenuExtender_SelectedAssets>& CBAssetMenuExtenderDelegates = ContentBrowserModule.GetAllAssetViewContextMenuExtenders();
+	CBAssetMenuExtenderDelegates.Add(FContentBrowserMenuExtender_SelectedAssets::CreateRaw(this, &FShotgridUIManagerImpl::OnExtendContentBrowserAssetSelectionMenu));
+	ContentBrowserAssetExtenderDelegateHandle = CBAssetMenuExtenderDelegates.Last().GetHandle();
+
+	// Register Level Editor menu extender
+	FLevelEditorModule& LevelEditorModule = FModuleManager::GetModuleChecked<FLevelEditorModule>(LEVELEDITOR_MODULE_NAME);
+
+	TArray<FLevelEditorModule::FLevelViewportMenuExtender_SelectedActors>& LevelEditorMenuExtenderDelegates = LevelEditorModule.GetAllLevelViewportContextMenuExtenders();
+	LevelEditorMenuExtenderDelegates.Add(FLevelEditorModule::FLevelViewportMenuExtender_SelectedActors::CreateRaw(this, &FShotgridUIManagerImpl::OnExtendLevelEditor));
+	LevelEditorExtenderDelegateHandle = LevelEditorMenuExtenderDelegates.Last().GetHandle();
+}
+
+void FShotgridUIManagerImpl::RemoveShotgridContextMenus()
+{
+	if (FModuleManager::Get().IsModuleLoaded(LEVELEDITOR_MODULE_NAME))
+	{
+		FLevelEditorModule& LevelEditorModule = FModuleManager::GetModuleChecked<FLevelEditorModule>(LEVELEDITOR_MODULE_NAME);
+		TArray<FLevelEditorModule::FLevelViewportMenuExtender_SelectedActors>& LevelEditorMenuExtenderDelegates = LevelEditorModule.GetAllLevelViewportContextMenuExtenders();
+		LevelEditorMenuExtenderDelegates.RemoveAll([this](const FLevelEditorModule::FLevelViewportMenuExtender_SelectedActors& Delegate) { return Delegate.GetHandle() == LevelEditorExtenderDelegateHandle; });
+	}
+
+	if (FModuleManager::Get().IsModuleLoaded(CONTENTBROWSER_MODULE_NAME))
+	{
+		FContentBrowserModule& ContentBrowserModule = FModuleManager::GetModuleChecked<FContentBrowserModule>(CONTENTBROWSER_MODULE_NAME);
+		TArray<FContentBrowserMenuExtender_SelectedAssets>& CBAssetMenuExtenderDelegates = ContentBrowserModule.GetAllAssetViewContextMenuExtenders();
+		CBAssetMenuExtenderDelegates.RemoveAll([this](const FContentBrowserMenuExtender_SelectedAssets& Delegate) { return Delegate.GetHandle() == ContentBrowserAssetExtenderDelegateHandle; });
+	}
+}
+
+void FShotgridUIManagerImpl::FillToolbar(FToolBarBuilder& ToolbarBuilder)
+{
+	ToolbarBuilder.BeginSection(TEXT("Shotgrid"));
+	{
+		// Add a drop-down menu (with a label and an icon for the drop-down button) to list the Shotgrid actions available
+		ToolbarBuilder.AddComboButton(
+			FUIAction(),
+			FOnGetContent::CreateRaw(this, &FShotgridUIManagerImpl::GenerateShotgridToolbarMenu),
+			LOCTEXT("ShotgridCombo_Label", "Shotgrid"),
+			LOCTEXT("ShotgridCombo_Tooltip", "Available Shotgrid commands"),
+			FSlateIcon(FShotgridStyle::GetStyleSetName(), "Shotgrid.Logo")
+		);
+	}
+	ToolbarBuilder.EndSection();
+}
+
+TSharedRef<SWidget> FShotgridUIManagerImpl::GenerateShotgridToolbarMenu()
+{
+	const bool bShouldCloseWindowAfterMenuSelection = true;
+	FMenuBuilder MenuBuilder(bShouldCloseWindowAfterMenuSelection, nullptr);
+
+	GenerateShotgridMenuContent(MenuBuilder, nullptr, nullptr);
+
+	return MenuBuilder.MakeWidget();
+}
+
+void FShotgridUIManagerImpl::GenerateShotgridMenuContent(FMenuBuilder& MenuBuilder, const TArray<FAssetData>* SelectedAssets, const TArray< AActor*>* SelectedActors)
+{
+	if (UShotgridEngine* Engine = UShotgridEngine::GetInstance())
+	{
+		Engine->SetSelection(SelectedAssets, SelectedActors);
+
+		// Query the available Shotgrid commands from the Shotgrid engine
+		TArray<FShotgridMenuItem> MenuItems = Engine->GetShotgridMenuItems();
+		for (const FShotgridMenuItem& MenuItem : MenuItems)
+		{
+			if (MenuItem.Type == TEXT("context_begin"))
+			{
+				MenuBuilder.BeginSection(NAME_None, FText::FromString(MenuItem.Title));
+			}
+			else if (MenuItem.Type == TEXT("context_end"))
+			{
+				MenuBuilder.EndSection();
+			}
+			else if (MenuItem.Type == TEXT("separator"))
+			{
+				MenuBuilder.AddMenuSeparator();
+			}
+			else
+			{
+				// The other menu types correspond to actual Shotgrid commands with an associated action
+				FString CommandName = MenuItem.Title;
+				MenuBuilder.AddMenuEntry(
+					FText::FromString(CommandName),
+					FText::FromString(MenuItem.Description),
+					FSlateIcon(),
+					FExecuteAction::CreateLambda([CommandName]()
+					{
+						if (UShotgridEngine* Engine = UShotgridEngine::GetInstance())
+						{
+							Engine->ExecuteCommand(CommandName);
+						}
+					})
+				);
+			}
+		}
+	}
+}
+
+void FShotgridUIManagerImpl::GenerateShotgridAssetContextMenu(FMenuBuilder& MenuBuilder, TArray<FAssetData> SelectedAssets)
+{
+	GenerateShotgridMenuContent(MenuBuilder, &SelectedAssets, nullptr);
+}
+
+void FShotgridUIManagerImpl::GenerateShotgridActorContextMenu(FMenuBuilder& MenuBuilder, TArray<AActor*> SelectedActors)
+{
+	GenerateShotgridMenuContent(MenuBuilder, nullptr, &SelectedActors);
+}
+
+TSharedRef<FExtender> FShotgridUIManagerImpl::OnExtendContentBrowserAssetSelectionMenu(const TArray<FAssetData>& SelectedAssets)
+{
+	// Menu extender for Content Browser context menu when an asset is selected
+	TSharedRef<FExtender> Extender(new FExtender());
+
+	if (SelectedAssets.Num() > 0)
+	{
+		Extender->AddMenuExtension("AssetContextReferences", EExtensionHook::After, nullptr, FMenuExtensionDelegate::CreateLambda(
+			[this, SelectedAssets](FMenuBuilder& MenuBuilder)
+			{
+				MenuBuilder.AddMenuSeparator();
+				MenuBuilder.AddSubMenu(
+					LOCTEXT("Shotgrid_ContextMenu", "Shotgrid"),
+					FText(),
+					FNewMenuDelegate::CreateRaw(this, &FShotgridUIManagerImpl::GenerateShotgridAssetContextMenu, SelectedAssets),
+					false,
+					FSlateIcon(FShotgridStyle::GetStyleSetName(), "Shotgrid.ContextLogo")
+				);
+			}));
+	}
+
+	return Extender;
+}
+
+TSharedRef<FExtender> FShotgridUIManagerImpl::OnExtendLevelEditor(const TSharedRef<FUICommandList> CommandList, const TArray<AActor*> SelectedActors)
+{
+	// Menu extender for Level Editor and World Outliner context menus when an actor is selected
+	TSharedRef<FExtender> Extender(new FExtender());
+
+	if (SelectedActors.Num() > 0)
+	{
+		Extender->AddMenuExtension("ActorUETools", EExtensionHook::After, nullptr, FMenuExtensionDelegate::CreateLambda(
+			[this, SelectedActors](FMenuBuilder& MenuBuilder)
+			{
+				MenuBuilder.AddSubMenu(
+					LOCTEXT("Shotgrid_ContextMenu", "Shotgrid"),
+					FText(),
+					FNewMenuDelegate::CreateRaw(this, &FShotgridUIManagerImpl::GenerateShotgridActorContextMenu, SelectedActors),
+					false,
+					FSlateIcon(FShotgridStyle::GetStyleSetName(), "Shotgrid.ContextLogo")
+				);
+			}));
+	}
+
+	return Extender;
+}
+
+void FShotgridUIManager::Initialize()
+{
+	if (!Instance.IsValid())
+	{
+		Instance = MakeUnique<FShotgridUIManagerImpl>();
+		Instance->Initialize();
+	}
+}
+
+void FShotgridUIManager::Shutdown()
+{
+	if (Instance.IsValid())
+	{
+		Instance->Shutdown();
+		Instance.Reset();
+	}
+}
+
+#undef LOCTEXT_NAMESPACE
diff --git a/Engine/Plugins/Experimental/Shotgrid/Source/Shotgrid/Private/ShotgridUIManager.h b/Engine/Plugins/Experimental/Shotgrid/Source/Shotgrid/Private/ShotgridUIManager.h
new file mode 100644
--- /dev/null
+++ b/Engine/Plugins/Experimental/Shotgrid/Source/Shotgrid/Private/ShotgridUIManager.h
@@ -0,0 +1,17 @@
+// Copyright Epic Games, Inc. All Rights Reserved.
+
+#pragma once
+
+#include "CoreMinimal.h"
+
+class FShotgridUIManagerImpl;
+
+class FShotgridUIManager
+{
+public:
+	static void Initialize();
+	static void Shutdown();
+
+private:
+	static TUniquePtr<FShotgridUIManagerImpl> Instance;
+};
\ No newline at end of file
diff --git a/Engine/Plugins/Experimental/Shotgrid/Source/Shotgrid/Public/IShotgridModule.h b/Engine/Plugins/Experimental/Shotgrid/Source/Shotgrid/Public/IShotgridModule.h
new file mode 100644
--- /dev/null
+++ b/Engine/Plugins/Experimental/Shotgrid/Source/Shotgrid/Public/IShotgridModule.h
@@ -0,0 +1,38 @@
+// Copyright Epic Games, Inc. All Rights Reserved.
+
+#pragma once
+
+#include "Modules/ModuleInterface.h"
+#include "Modules/ModuleManager.h"		// For inline LoadModuleChecked()
+
+#define SHOTGUN_MODULE_NAME TEXT("Shotgrid")
+
+/**
+ * The public interface of the Shotgrid module
+ */
+class IShotgridModule : public IModuleInterface
+{
+
+public:
+
+	/**
+	 * Singleton-like access to IShotgridPlugin
+	 *
+	 * @return Returns Shotplugin singleton instance, loading the module on demand if needed
+	 */
+	static inline IShotgridModule& Get()
+	{
+		return FModuleManager::LoadModuleChecked<IShotgridModule>(SHOTGUN_MODULE_NAME);
+	}
+
+	/**
+	 * Checks to see if this module is loaded and ready.  It is only valid to call Get() if IsAvailable() returns true.
+	 *
+	 * @return True if the module is loaded and ready to use
+	 */
+	static inline bool IsAvailable()
+	{
+		return FModuleManager::Get().IsModuleLoaded(SHOTGUN_MODULE_NAME);
+	}
+};
+
diff --git a/Engine/Plugins/Experimental/Shotgrid/Source/Shotgrid/Shotgrid.Build.cs b/Engine/Plugins/Experimental/Shotgrid/Source/Shotgrid/Shotgrid.Build.cs
new file mode 100644
--- /dev/null
+++ b/Engine/Plugins/Experimental/Shotgrid/Source/Shotgrid/Shotgrid.Build.cs
@@ -0,0 +1,30 @@
+// Copyright Epic Games, Inc. All Rights Reserved.
+
+namespace UnrealBuildTool.Rules
+{
+	public class Shotgrid : ModuleRules
+	{
+		public Shotgrid(ReadOnlyTargetRules Target)
+			: base(Target)
+		{
+			PCHUsage = PCHUsageMode.UseExplicitOrSharedPCHs;
+
+			PrivateDependencyModuleNames.AddRange(
+				new string[]
+				{
+					"AssetRegistry",
+					"ContentBrowser",
+					"Core",
+					"CoreUObject",
+					"EditorStyle",
+					"Engine",
+					"LevelEditor",
+					"PythonScriptPlugin",
+					"Settings",
+					"Slate",
+					"SlateCore",
+				}
+			);
+		}
+	}
+}
\ No newline at end of file
diff --git a/Engine/Plugins/Experimental/Shotgun/Shotgun.uplugin b/Engine/Plugins/Experimental/Shotgun/Shotgun.uplugin
deleted file mode 100644
diff --git a/Engine/Plugins/Experimental/Shotgun/Source/Shotgun/Private/ShotgunEngine.cpp b/Engine/Plugins/Experimental/Shotgun/Source/Shotgun/Private/ShotgunEngine.cpp
deleted file mode 100644
diff --git a/Engine/Plugins/Experimental/Shotgun/Source/Shotgun/Private/ShotgunEngine.h b/Engine/Plugins/Experimental/Shotgun/Source/Shotgun/Private/ShotgunEngine.h
deleted file mode 100644
diff --git a/Engine/Plugins/Experimental/Shotgun/Source/Shotgun/Private/ShotgunModule.cpp b/Engine/Plugins/Experimental/Shotgun/Source/Shotgun/Private/ShotgunModule.cpp
deleted file mode 100644
diff --git a/Engine/Plugins/Experimental/Shotgun/Source/Shotgun/Private/ShotgunSettings.cpp b/Engine/Plugins/Experimental/Shotgun/Source/Shotgun/Private/ShotgunSettings.cpp
deleted file mode 100644
diff --git a/Engine/Plugins/Experimental/Shotgun/Source/Shotgun/Private/ShotgunSettings.h b/Engine/Plugins/Experimental/Shotgun/Source/Shotgun/Private/ShotgunSettings.h
deleted file mode 100644
diff --git a/Engine/Plugins/Experimental/Shotgun/Source/Shotgun/Private/ShotgunStyle.cpp b/Engine/Plugins/Experimental/Shotgun/Source/Shotgun/Private/ShotgunStyle.cpp
deleted file mode 100644
diff --git a/Engine/Plugins/Experimental/Shotgun/Source/Shotgun/Private/ShotgunStyle.h b/Engine/Plugins/Experimental/Shotgun/Source/Shotgun/Private/ShotgunStyle.h
deleted file mode 100644
diff --git a/Engine/Plugins/Experimental/Shotgun/Source/Shotgun/Private/ShotgunUIManager.cpp b/Engine/Plugins/Experimental/Shotgun/Source/Shotgun/Private/ShotgunUIManager.cpp
deleted file mode 100644
diff --git a/Engine/Plugins/Experimental/Shotgun/Source/Shotgun/Private/ShotgunUIManager.h b/Engine/Plugins/Experimental/Shotgun/Source/Shotgun/Private/ShotgunUIManager.h
deleted file mode 100644
diff --git a/Engine/Plugins/Experimental/Shotgun/Source/Shotgun/Public/IShotgunModule.h b/Engine/Plugins/Experimental/Shotgun/Source/Shotgun/Public/IShotgunModule.h
deleted file mode 100644
diff --git a/Engine/Plugins/Experimental/Shotgun/Source/Shotgun/Shotgun.Build.cs b/Engine/Plugins/Experimental/Shotgun/Source/Shotgun/Shotgun.Build.cs
deleted file mode 100644
