Remove unnecessary line from CL

diff --git a/Engine/Plugins/Experimental/Shotgrid/Config/DefaultShotgrid.ini b/Engine/Plugins/Experimental/Shotgrid/Config/DefaultShotgrid.ini
--- a/Engine/Plugins/Experimental/Shotgrid/Config/DefaultShotgrid.ini
+++ b/Engine/Plugins/Experimental/Shotgrid/Config/DefaultShotgrid.ini
@@ -1,5 +1,3 @@
-[/Script/Engine.Engine]
-
 [CoreRedirects]
 +ClassRedirects=(OldName="Shotgun",NewName="Shotgrid",MatchSubstring=true)
 +EnumRedirects=(OldName="Shotgun",NewName="Shotgrid",MatchSubstring=true)
