Clean-up redirects per the post review.

diff --git a/Engine/Plugins/Experimental/Shotgrid/Config/DefaultShotgrid.ini b/Engine/Plugins/Experimental/Shotgrid/Config/DefaultShotgrid.ini
--- a/Engine/Plugins/Experimental/Shotgrid/Config/DefaultShotgrid.ini
+++ b/Engine/Plugins/Experimental/Shotgrid/Config/DefaultShotgrid.ini
@@ -1,6 +1,5 @@
 [CoreRedirects]
-+ClassRedirects=(OldName="Shotgun",NewName="Shotgrid",MatchSubstring=true)
-+EnumRedirects=(OldName="Shotgun",NewName="Shotgrid",MatchSubstring=true)
-+FunctionRedirects=(OldName="Shotgun",NewName="Shotgrid",MatchSubstring=true)
-+StructRedirects=(OldName="Shotgun",NewName="Shotgrid",MatchSubstring=true)
-+PackageRedirects=(OldName="Shotgun",NewName="/Shotgrid/",MatchSubstring=true)
\ No newline at end of file
++ClassRedirects=(OldName="ShotgunEngine", NewName="/Script/Shotgrid.ShotgridEngine")
++ClassRedirects=(OldName="ShotgunSettings", NewName="/Script/Shotgrid.ShotgridSettings")
++PackageRedirects=(OldName="/Script/Shotgun",NewName="/Script/Shotgrid")
+
