Rename UE4 to Unreal Editor

diff --git a/Engine/Plugins/Experimental/Shotgrid/Shotgrid.uplugin b/Engine/Plugins/Experimental/Shotgrid/Shotgrid.uplugin
--- a/Engine/Plugins/Experimental/Shotgrid/Shotgrid.uplugin
+++ b/Engine/Plugins/Experimental/Shotgrid/Shotgrid.uplugin
@@ -3,7 +3,7 @@
 	"Version" : 1,
 	"VersionName" : "0.10",
 	"FriendlyName" : "Shotgrid Plugin",
-	"Description" : "Shotgrid integration for the UE4 Editor.",
+	"Description" : "Shotgrid integration for the Unreal Editor.",
 	"Category" : "Editor",
 	"CreatedBy" : "Epic Games, Inc.",
 	"CreatedByURL" : "http://epicgames.com",
@@ -41,4 +41,4 @@
 			"Enabled": true
 		}
 	]
-}
\ No newline at end of file
+}
