Add ShotGrid to the UE5 toolbar.

diff --git a/Engine/Plugins/Experimental/Shotgrid/Source/Shotgrid/Private/ShotgridUIManager.cpp b/Engine/Plugins/Experimental/Shotgrid/Source/Shotgrid/Private/ShotgridUIManager.cpp
--- a/Engine/Plugins/Experimental/Shotgrid/Source/Shotgrid/Private/ShotgridUIManager.cpp
+++ b/Engine/Plugins/Experimental/Shotgrid/Source/Shotgrid/Private/ShotgridUIManager.cpp
@@ -4,6 +4,12 @@
 #include "ShotgridEngine.h"
 #include "ShotgridStyle.h"
 
+#include "Framework/Application/SlateApplication.h"
+#include "Framework/Commands/Commands.h"
+#include "Framework/Commands/UIAction.h"
+#include "Framework/MultiBox/MultiBoxBuilder.h"
+#include "ToolMenus.h"
+
 #include "AssetData.h"
 #include "ContentBrowserModule.h"
 #include "LevelEditor.h"
@@ -15,9 +21,128 @@
 #include "IPythonScriptPlugin.h"
 
 #define LOCTEXT_NAMESPACE "Shotgrid"
+
 #define LEVELEDITOR_MODULE_NAME TEXT("LevelEditor")
 #define CONTENTBROWSER_MODULE_NAME TEXT("ContentBrowser")
 
+namespace UE::Private::Shotgrid
+{
+
+void GenerateShotgridMenuContent(FMenuBuilder& MenuBuilder, const TArray<FAssetData>* SelectedAssets, const TArray< AActor*>* SelectedActors)
+{
+	if (UShotgridEngine* Engine = UShotgridEngine::GetInstance())
+	{
+		Engine->SetSelection(SelectedAssets, SelectedActors);
+
+		// Query the available Shotgrid commands from the Shotgrid engine
+		TArray<FShotgridMenuItem> MenuItems = Engine->GetShotgridMenuItems();
+		for (const FShotgridMenuItem& MenuItem : MenuItems)
+		{
+			if (MenuItem.Type == TEXT("context_begin"))
+			{
+				MenuBuilder.BeginSection(NAME_None, FText::FromString(MenuItem.Title));
+			}
+			else if (MenuItem.Type == TEXT("context_end"))
+			{
+				MenuBuilder.EndSection();
+			}
+			else if (MenuItem.Type == TEXT("separator"))
+			{
+				MenuBuilder.AddMenuSeparator();
+			}
+			else
+			{
+				// The other menu types correspond to actual Shotgrid commands with an associated action
+				FString CommandName = MenuItem.Title;
+				MenuBuilder.AddMenuEntry(
+					FText::FromString(CommandName),
+					FText::FromString(MenuItem.Description),
+					FSlateIcon(),
+					FExecuteAction::CreateLambda([CommandName]()
+						{
+							if (UShotgridEngine* Engine = UShotgridEngine::GetInstance())
+							{
+								Engine->ExecuteCommand(CommandName);
+							}
+						})
+					);
+			}
+		}
+	}
+}
+
+bool IsEnabledWithEnvironmentVariable(FString EnvironmentVariableBase)
+{
+	const FString EnvBootstrapString = "UE_" + EnvironmentVariableBase + "_BOOTSTRAP";
+	const FString EnvEngine = EnvironmentVariableBase + "_ENGINE";
+	const FString EnvEntityType = EnvironmentVariableBase + "_ENTITY_TYPE";
+	const FString EnvEntityId = EnvironmentVariableBase + "_ENTITY_ID";
+	const FString EnvEnabled = "UE_" + EnvironmentVariableBase + "_ENABLED";
+
+	bool bIsEnabled = false;
+	// Check if the bootstrap environment variable is set and that the script exists
+	FString ShotgridBootstrap = FPlatformMisc::GetEnvironmentVariable(*EnvBootstrapString);
+	if (!ShotgridBootstrap.IsEmpty() && FPaths::FileExists(ShotgridBootstrap))
+	{
+		// The following environment variables must be set for the Shotgrid apps to be fully functional
+		// These variables are automatically set when the editor is launched through Shotgrid Desktop
+		FString ShotgridEngine = FPlatformMisc::GetEnvironmentVariable(*EnvEngine);
+		FString ShotgridEntityType = FPlatformMisc::GetEnvironmentVariable(*EnvEntityType);
+		FString ShotgridEntityId = FPlatformMisc::GetEnvironmentVariable(*EnvEntityId);
+
+		if (ShotgridEngine == TEXT("tk-unreal") && !ShotgridEntityType.IsEmpty() && !ShotgridEntityId.IsEmpty())
+		{
+			bIsEnabled = true;
+
+			// Set environment variable in the Python interpreter to enable the Shotgrid Unreal init script
+			FString ExecCommand = FString::Printf(TEXT("import os\nos.environ['%s']='True'"), *EnvEnabled);
+			IPythonScriptPlugin::Get()->ExecPythonCommand(*ExecCommand);
+		}
+	}
+	return bIsEnabled;
+}
+
+}
+
+struct FShotgridMenuEntryImpl
+{
+	FShotgridMenuEntryImpl()
+	{
+		UToolMenu* Menu = UToolMenus::Get()->ExtendMenu("LevelEditor.LevelEditorToolBar.User");
+		FToolMenuSection& Section = Menu->FindOrAddSection("Shotgrid");
+
+		const FToolMenuEntry ShotgridComboEntry = FToolMenuEntry::InitComboButton(
+			TEXT("Shotgrid"),
+			FUIAction(),
+			FOnGetContent::CreateRaw(this, &FShotgridMenuEntryImpl::GenerateShotgridToolbarMenu),
+			LOCTEXT("ShotgridCombo_Label", "Shotgrid"),
+			LOCTEXT("ShotgridCombo_Tooltip", "Available Shotgrid commands"),
+			FSlateIcon(FShotgridStyle::GetStyleSetName(), "Shotgrid.Logo"));
+
+		Section.AddEntry(ShotgridComboEntry);
+
+	}
+
+	~FShotgridMenuEntryImpl()
+	{
+		if (!IsEngineExitRequested() && UObjectInitialized())
+		{
+			UToolMenus::Get()->RemoveSection("LevelEditor.LevelEditorToolBar.User", "Shotgrid");
+		}
+	}
+
+	TSharedRef<SWidget> GenerateShotgridToolbarMenu()
+	{
+		const bool bShouldCloseWindowAfterMenuSelection = true;
+		FMenuBuilder MenuBuilder(bShouldCloseWindowAfterMenuSelection, nullptr);
+
+		UE::Private::Shotgrid::GenerateShotgridMenuContent(MenuBuilder, nullptr, nullptr);
+
+		return MenuBuilder.MakeWidget();
+	}
+
+};
+
 TUniquePtr<FShotgridUIManagerImpl> FShotgridUIManager::Instance;
 
 class FShotgridUIManagerImpl
@@ -27,54 +152,45 @@ public:
 	void Shutdown();
 
 private:
-	FDelegateHandle ContentBrowserAssetExtenderDelegateHandle;
-	FDelegateHandle LevelEditorExtenderDelegateHandle;
-
-	bool bIsShotgridEnabled;
-
-	void SetupShotgridMenu();
 	void SetupShotgridContextMenus();
 	void RemoveShotgridContextMenus();
 
 	void FillToolbar(FToolBarBuilder& ToolbarBuilder);
 	TSharedRef<SWidget> GenerateShotgridToolbarMenu();
-	void GenerateShotgridMenuContent(FMenuBuilder& MenuBuilder, const TArray<FAssetData>* SelectedAssets, const TArray< AActor*>* SelectedActors);
+
 	void GenerateShotgridAssetContextMenu(FMenuBuilder& MenuBuilder, TArray<FAssetData> SelectedAssets);
 	void GenerateShotgridActorContextMenu(FMenuBuilder& MenuBuilder, TArray< AActor*> SelectedActors);
 
 	// Menu extender callbacks
 	TSharedRef<FExtender> OnExtendLevelEditor(const TSharedRef<FUICommandList> CommandList, const TArray<AActor*> SelectedActors);
 	TSharedRef<FExtender> OnExtendContentBrowserAssetSelectionMenu(const TArray<FAssetData>& SelectedAssets);
-};
 
-void FShotgridUIManagerImpl::Initialize()
-{
-	bIsShotgridEnabled = false;
+	FDelegateHandle ContentBrowserAssetExtenderDelegateHandle;
+	FDelegateHandle LevelEditorExtenderDelegateHandle;
 
-	// Check if the bootstrap environment variable is set and that the script exists
-	FString ShotgridBootstrap = FPlatformMisc::GetEnvironmentVariable(TEXT("UE_SHOTGUN_BOOTSTRAP"));
-	if (!ShotgridBootstrap.IsEmpty() && FPaths::FileExists(ShotgridBootstrap))
-	{
-		// The following environment variables must be set for the Shotgrid apps to be fully functional
-		// These variables are automatically set when the editor is launched through Shotgrid Desktop
-		FString ShotgridEngine = FPlatformMisc::GetEnvironmentVariable(TEXT("SHOTGUN_ENGINE"));
-		FString ShotgridEntityType = FPlatformMisc::GetEnvironmentVariable(TEXT("SHOTGUN_ENTITY_TYPE"));
-		FString ShotgridEntityId = FPlatformMisc::GetEnvironmentVariable(TEXT("SHOTGUN_ENTITY_ID"));
+	bool bIsShotgridEnabled = false;
 
-		if (ShotgridEngine == TEXT("tk-unreal") && !ShotgridEntityType.IsEmpty() && !ShotgridEntityId.IsEmpty())
-		{
-			bIsShotgridEnabled = true;
+	TUniquePtr<FShotgridMenuEntryImpl> Toolbar;
+};
 
-			// Set environment variable in the Python interpreter to enable the Shotgrid Unreal init script
-			IPythonScriptPlugin::Get()->ExecPythonCommand(TEXT("import os\nos.environ['UE_SHOTGUN_ENABLED']='True'"));
-		}
-	}
+
+void FShotgridUIManagerImpl::Initialize()
+{
+	bIsShotgridEnabled = UE::Private::Shotgrid::IsEnabledWithEnvironmentVariable(TEXT("SHOTGUN"))
+		|| UE::Private::Shotgrid::IsEnabledWithEnvironmentVariable(TEXT("SHOTGRID"));
 
 	if (bIsShotgridEnabled)
 	{
 		FShotgridStyle::Initialize();
 
-		SetupShotgridMenu();
+		// Set the Shotgrid icons
+		FShotgridStyle::SetIcon("Logo", "sg_logo_80px");
+		FShotgridStyle::SetIcon("ContextLogo", "sg_context_logo");
+
+		if (!IsRunningCommandlet())
+		{
+			Toolbar = MakeUnique<FShotgridMenuEntryImpl>();
+		}
 		SetupShotgridContextMenus();
 	}
 }
@@ -83,27 +199,13 @@ void FShotgridUIManagerImpl::Shutdown()
 {
 	if (bIsShotgridEnabled)
 	{
+		Toolbar = {};
 		RemoveShotgridContextMenus();
 
 		FShotgridStyle::Shutdown();
 	}
 }
 
-void FShotgridUIManagerImpl::SetupShotgridMenu()
-{
-	// Set the Shotgrid icons
-	FShotgridStyle::SetIcon("Logo", "sg_logo_80px");
-	FShotgridStyle::SetIcon("ContextLogo", "sg_context_logo");
-
-	// Add a Shotgrid toolbar section after the settings section of the level editor
-	FLevelEditorModule& LevelEditorModule = FModuleManager::LoadModuleChecked<FLevelEditorModule>(LEVELEDITOR_MODULE_NAME);
-
-	TSharedPtr<FExtender> ToolbarExtender = MakeShareable(new FExtender);
-	ToolbarExtender->AddToolBarExtension("Settings", EExtensionHook::After, nullptr, FToolBarExtensionDelegate::CreateRaw(this, &FShotgridUIManagerImpl::FillToolbar));
-
-	LevelEditorModule.GetToolBarExtensibilityManager()->AddExtender(ToolbarExtender);
-}
-
 void FShotgridUIManagerImpl::SetupShotgridContextMenus()
 {
 	// Register Content Browser menu extender
@@ -138,83 +240,14 @@ void FShotgridUIManagerImpl::RemoveShotgridContextMenus()
 	}
 }
 
-void FShotgridUIManagerImpl::FillToolbar(FToolBarBuilder& ToolbarBuilder)
-{
-	ToolbarBuilder.BeginSection(TEXT("Shotgrid"));
-	{
-		// Add a drop-down menu (with a label and an icon for the drop-down button) to list the Shotgrid actions available
-		ToolbarBuilder.AddComboButton(
-			FUIAction(),
-			FOnGetContent::CreateRaw(this, &FShotgridUIManagerImpl::GenerateShotgridToolbarMenu),
-			LOCTEXT("ShotgridCombo_Label", "Shotgrid"),
-			LOCTEXT("ShotgridCombo_Tooltip", "Available Shotgrid commands"),
-			FSlateIcon(FShotgridStyle::GetStyleSetName(), "Shotgrid.Logo")
-		);
-	}
-	ToolbarBuilder.EndSection();
-}
-
-TSharedRef<SWidget> FShotgridUIManagerImpl::GenerateShotgridToolbarMenu()
-{
-	const bool bShouldCloseWindowAfterMenuSelection = true;
-	FMenuBuilder MenuBuilder(bShouldCloseWindowAfterMenuSelection, nullptr);
-
-	GenerateShotgridMenuContent(MenuBuilder, nullptr, nullptr);
-
-	return MenuBuilder.MakeWidget();
-}
-
-void FShotgridUIManagerImpl::GenerateShotgridMenuContent(FMenuBuilder& MenuBuilder, const TArray<FAssetData>* SelectedAssets, const TArray< AActor*>* SelectedActors)
-{
-	if (UShotgridEngine* Engine = UShotgridEngine::GetInstance())
-	{
-		Engine->SetSelection(SelectedAssets, SelectedActors);
-
-		// Query the available Shotgrid commands from the Shotgrid engine
-		TArray<FShotgridMenuItem> MenuItems = Engine->GetShotgridMenuItems();
-		for (const FShotgridMenuItem& MenuItem : MenuItems)
-		{
-			if (MenuItem.Type == TEXT("context_begin"))
-			{
-				MenuBuilder.BeginSection(NAME_None, FText::FromString(MenuItem.Title));
-			}
-			else if (MenuItem.Type == TEXT("context_end"))
-			{
-				MenuBuilder.EndSection();
-			}
-			else if (MenuItem.Type == TEXT("separator"))
-			{
-				MenuBuilder.AddMenuSeparator();
-			}
-			else
-			{
-				// The other menu types correspond to actual Shotgrid commands with an associated action
-				FString CommandName = MenuItem.Title;
-				MenuBuilder.AddMenuEntry(
-					FText::FromString(CommandName),
-					FText::FromString(MenuItem.Description),
-					FSlateIcon(),
-					FExecuteAction::CreateLambda([CommandName]()
-					{
-						if (UShotgridEngine* Engine = UShotgridEngine::GetInstance())
-						{
-							Engine->ExecuteCommand(CommandName);
-						}
-					})
-				);
-			}
-		}
-	}
-}
-
 void FShotgridUIManagerImpl::GenerateShotgridAssetContextMenu(FMenuBuilder& MenuBuilder, TArray<FAssetData> SelectedAssets)
 {
-	GenerateShotgridMenuContent(MenuBuilder, &SelectedAssets, nullptr);
+	UE::Private::Shotgrid::GenerateShotgridMenuContent(MenuBuilder, &SelectedAssets, nullptr);
 }
 
 void FShotgridUIManagerImpl::GenerateShotgridActorContextMenu(FMenuBuilder& MenuBuilder, TArray<AActor*> SelectedActors)
 {
-	GenerateShotgridMenuContent(MenuBuilder, nullptr, &SelectedActors);
+	UE::Private::Shotgrid::GenerateShotgridMenuContent(MenuBuilder, nullptr, &SelectedActors);
 }
 
 TSharedRef<FExtender> FShotgridUIManagerImpl::OnExtendContentBrowserAssetSelectionMenu(const TArray<FAssetData>& SelectedAssets)
diff --git a/Engine/Plugins/Experimental/Shotgrid/Source/Shotgrid/Shotgrid.Build.cs b/Engine/Plugins/Experimental/Shotgrid/Source/Shotgrid/Shotgrid.Build.cs
--- a/Engine/Plugins/Experimental/Shotgrid/Source/Shotgrid/Shotgrid.Build.cs
+++ b/Engine/Plugins/Experimental/Shotgrid/Source/Shotgrid/Shotgrid.Build.cs
@@ -23,8 +23,10 @@ namespace UnrealBuildTool.Rules
 					"Settings",
 					"Slate",
 					"SlateCore",
+					"ToolMenus",
+					"UnrealEd"
 				}
 			);
 		}
 	}
-}
\ No newline at end of file
+}
