Replace icons for Shotgrid plugin

diff --git a/Engine/Plugins/Experimental/Shotgrid/Source/Shotgrid/Private/ShotgridUIManager.cpp b/Engine/Plugins/Experimental/Shotgrid/Source/Shotgrid/Private/ShotgridUIManager.cpp
--- a/Engine/Plugins/Experimental/Shotgrid/Source/Shotgrid/Private/ShotgridUIManager.cpp
+++ b/Engine/Plugins/Experimental/Shotgrid/Source/Shotgrid/Private/ShotgridUIManager.cpp
@@ -184,8 +184,8 @@ void FShotgridUIManagerImpl::Initialize()
 		FShotgridStyle::Initialize();
 
 		// Set the Shotgrid icons
-		FShotgridStyle::SetIcon("Logo", "sg_logo_80px");
-		FShotgridStyle::SetIcon("ContextLogo", "sg_context_logo");
+		FShotgridStyle::SetIcon("Logo", "W18_SG_QAT_40x40");
+		FShotgridStyle::SetIcon("ContextLogo", "W20_SG_QAT_24x24");
 
 		if (!IsRunningCommandlet())
 		{
