ShotGrid: Capitalize G in user-facing instances of "ShotGrid."

diff --git a/Engine/Plugins/Experimental/Shotgrid/Shotgrid.uplugin b/Engine/Plugins/Experimental/Shotgrid/Shotgrid.uplugin
--- a/Engine/Plugins/Experimental/Shotgrid/Shotgrid.uplugin
+++ b/Engine/Plugins/Experimental/Shotgrid/Shotgrid.uplugin
@@ -2,8 +2,8 @@
 	"FileVersion" : 1,
 	"Version" : 1,
 	"VersionName" : "0.10",
-	"FriendlyName" : "Shotgrid Plugin",
-	"Description" : "Shotgrid integration for the Unreal Editor.",
+	"FriendlyName" : "ShotGrid Plugin",
+	"Description" : "ShotGrid integration for the Unreal Editor.",
 	"Category" : "Editor",
 	"CreatedBy" : "Epic Games, Inc.",
 	"CreatedByURL" : "http://epicgames.com",
diff --git a/Engine/Plugins/Experimental/Shotgrid/Source/Shotgrid/Private/ShotgridModule.cpp b/Engine/Plugins/Experimental/Shotgrid/Source/Shotgrid/Private/ShotgridModule.cpp
--- a/Engine/Plugins/Experimental/Shotgrid/Source/Shotgrid/Private/ShotgridModule.cpp
+++ b/Engine/Plugins/Experimental/Shotgrid/Source/Shotgrid/Private/ShotgridModule.cpp
@@ -39,8 +39,8 @@ protected:
 		if (SettingsModule != nullptr)
 		{
 			ISettingsSectionPtr SettingsSection = SettingsModule->RegisterSettings("Project", "Plugins", "Shotgrid",
-				LOCTEXT("ShotgridSettingsName", "Shotgrid"),
-				LOCTEXT("ShotgridSettingsDescription", "Configure the Shotgrid plug-in."),
+				LOCTEXT("ShotgridSettingsName", "ShotGrid"),
+				LOCTEXT("ShotgridSettingsDescription", "Configure the ShotGrid plugin."),
 				GetMutableDefault<UShotgridSettings>()
 			);
 		}
diff --git a/Engine/Plugins/Experimental/Shotgrid/Source/Shotgrid/Private/ShotgridUIManager.cpp b/Engine/Plugins/Experimental/Shotgrid/Source/Shotgrid/Private/ShotgridUIManager.cpp
--- a/Engine/Plugins/Experimental/Shotgrid/Source/Shotgrid/Private/ShotgridUIManager.cpp
+++ b/Engine/Plugins/Experimental/Shotgrid/Source/Shotgrid/Private/ShotgridUIManager.cpp
@@ -115,8 +115,8 @@ struct FShotgridMenuEntryImpl
 			TEXT("Shotgrid"),
 			FUIAction(),
 			FOnGetContent::CreateRaw(this, &FShotgridMenuEntryImpl::GenerateShotgridToolbarMenu),
-			LOCTEXT("ShotgridCombo_Label", "Shotgrid"),
-			LOCTEXT("ShotgridCombo_Tooltip", "Available Shotgrid commands"),
+			LOCTEXT("ShotgridCombo_Label", "ShotGrid"),
+			LOCTEXT("ShotgridCombo_Tooltip", "Available ShotGrid commands"),
 			FSlateIcon(FShotgridStyle::GetStyleSetName(), "Shotgrid.Logo"));
 
 		Section.AddEntry(ShotgridComboEntry);
@@ -262,7 +262,7 @@ TSharedRef<FExtender> FShotgridUIManagerImpl::OnExtendContentBrowserAssetSelecti
 			{
 				MenuBuilder.AddMenuSeparator();
 				MenuBuilder.AddSubMenu(
-					LOCTEXT("Shotgrid_ContextMenu", "Shotgrid"),
+					LOCTEXT("Shotgrid_ContextMenu", "ShotGrid"),
 					FText(),
 					FNewMenuDelegate::CreateRaw(this, &FShotgridUIManagerImpl::GenerateShotgridAssetContextMenu, SelectedAssets),
 					false,
@@ -285,7 +285,7 @@ TSharedRef<FExtender> FShotgridUIManagerImpl::OnExtendLevelEditor(const TSharedR
 			[this, SelectedActors](FMenuBuilder& MenuBuilder)
 			{
 				MenuBuilder.AddSubMenu(
-					LOCTEXT("Shotgrid_ContextMenu", "Shotgrid"),
+					LOCTEXT("Shotgrid_ContextMenu", "ShotGrid"),
 					FText(),
 					FNewMenuDelegate::CreateRaw(this, &FShotgridUIManagerImpl::GenerateShotgridActorContextMenu, SelectedActors),
 					false,
