Enable ShotGrid on Linux

diff --git a/Engine/Plugins/Experimental/Shotgrid/Shotgrid.uplugin b/Engine/Plugins/Experimental/Shotgrid/Shotgrid.uplugin
--- a/Engine/Plugins/Experimental/Shotgrid/Shotgrid.uplugin
+++ b/Engine/Plugins/Experimental/Shotgrid/Shotgrid.uplugin
@@ -14,10 +14,6 @@
 	"CanContainContent" : true,
 	"IsBetaVersion" : true,
 	"Installed" : false,
-	"SupportedTargetPlatforms":
-	[
-		"Win64"
-	],
 	"Modules":
 	[
 		{
@@ -26,7 +22,8 @@
 			"LoadingPhase": "Default",
 			"PlatformAllowList":
 			[
-				"Win64"
+				"Win64",
+				"Linux"
 			]
 		}
 	],
