Rename shotgrid to Flow Production Tracking.  Fix startup issue with Flow trying to run before python finishes loading.

diff --git a/Engine/Plugins/Experimental/Shotgrid/Shotgrid.uplugin b/Engine/Plugins/Experimental/Shotgrid/Shotgrid.uplugin
--- a/Engine/Plugins/Experimental/Shotgrid/Shotgrid.uplugin
+++ b/Engine/Plugins/Experimental/Shotgrid/Shotgrid.uplugin
@@ -2,8 +2,8 @@
 	"FileVersion" : 1,
 	"Version" : 1,
 	"VersionName" : "0.10",
-	"FriendlyName" : "ShotGrid Plugin",
-	"Description" : "ShotGrid integration for the Unreal Editor.",
+	"FriendlyName" : "Flow Production Tracking",
+	"Description" : "Flow Production Tracking (formerly known as Shotgun and/or ShotGrid) integration for the Unreal Editor.",
 	"Category" : "Editor",
 	"CreatedBy" : "Epic Games, Inc.",
 	"CreatedByURL" : "https://epicgames.com",
diff --git a/Engine/Plugins/Experimental/Shotgrid/Source/Shotgrid/Private/ShotgridModule.cpp b/Engine/Plugins/Experimental/Shotgrid/Source/Shotgrid/Private/ShotgridModule.cpp
--- a/Engine/Plugins/Experimental/Shotgrid/Source/Shotgrid/Private/ShotgridModule.cpp
+++ b/Engine/Plugins/Experimental/Shotgrid/Source/Shotgrid/Private/ShotgridModule.cpp
@@ -3,7 +3,7 @@
 #include "IShotgridModule.h"
 #include "ShotgridSettings.h"
 #include "ShotgridUIManager.h"
-
+#include "Misc/CoreDelegates.h"
 #include "ISettingsModule.h"
 #include "ISettingsSection.h"
 
@@ -12,28 +12,35 @@
 class FShotgridModule : public IShotgridModule
 {
 public:
+
+	static void OnEngineStartupComplete()
+	{
+		RegisterSettings();
+
+		FShotgridUIManager::Initialize();
+	}
+
+	static void OnEnginePreExit()
+	{
+		FShotgridUIManager::Shutdown();
+
+		UnregisterSettings();
+	}
 	virtual void StartupModule() override
 	{
 		if (GIsEditor && !IsRunningCommandlet())
 		{
-			RegisterSettings();
-
-			FShotgridUIManager::Initialize();
+			FCoreDelegates::OnFEngineLoopInitComplete.AddStatic(&OnEngineStartupComplete);
+			FCoreDelegates::OnEnginePreExit.AddStatic(&OnEnginePreExit);
 		}
 	}
 
 	virtual void ShutdownModule() override
 	{
-		if ((GIsEditor && !IsRunningCommandlet()))
-		{
-			FShotgridUIManager::Shutdown();
-
-			UnregisterSettings();
-		}
 	}
 
 protected:
-	void RegisterSettings()
+	static void RegisterSettings()
 	{
 		ISettingsModule* SettingsModule = FModuleManager::GetModulePtr<ISettingsModule>("Settings");
 		if (SettingsModule != nullptr)
@@ -46,7 +53,7 @@ protected:
 		}
 	}
 
-	void UnregisterSettings()
+	static void UnregisterSettings()
 	{
 		ISettingsModule* SettingsModule = FModuleManager::GetModulePtr<ISettingsModule>("Settings");
 		if (SettingsModule != nullptr)
diff --git a/Engine/Plugins/Experimental/Shotgrid/Source/Shotgrid/Private/ShotgridUIManager.cpp b/Engine/Plugins/Experimental/Shotgrid/Source/Shotgrid/Private/ShotgridUIManager.cpp
--- a/Engine/Plugins/Experimental/Shotgrid/Source/Shotgrid/Private/ShotgridUIManager.cpp
+++ b/Engine/Plugins/Experimental/Shotgrid/Source/Shotgrid/Private/ShotgridUIManager.cpp
@@ -110,8 +110,8 @@ struct FShotgridMenuEntryImpl
 			TEXT("Shotgrid"),
 			FUIAction(),
 			FOnGetContent::CreateRaw(this, &FShotgridMenuEntryImpl::GenerateShotgridToolbarMenu),
-			LOCTEXT("ShotgridCombo_Label", "ShotGrid"),
-			LOCTEXT("ShotgridCombo_Tooltip", "Available ShotGrid commands"),
+			LOCTEXT("ShotgridCombo_Label", "Flow Production Tracking"),
+			LOCTEXT("ShotgridCombo_Tooltip", "Available Flow Production commands"),
 			FSlateIcon(FShotgridStyle::GetStyleSetName(), "Shotgrid.Logo"));
 
 		Section.AddEntry(ShotgridComboEntry);
